Open World Template Upload. Contains map file with landscape, basic lighting, material with texture.

diff --git a/Templates/TP_OpenWorld/Config/DefaultEngine.ini b/Templates/TP_OpenWorld/Config/DefaultEngine.ini
new file mode 100644
--- /dev/null
+++ b/Templates/TP_OpenWorld/Config/DefaultEngine.ini
@@ -0,0 +1,28 @@
+
+
+[/Script/Engine.RendererSettings]
+r.DefaultFeature.AutoExposure.ExtendDefaultLuminanceRange=True
+
+r.GenerateMeshDistanceFields=True
+
+r.DynamicGlobalIlluminationMethod=1
+
+r.ReflectionMethod=1
+
+r.Shadow.Virtual.Enable=1
+
+[/Script/HardwareTargeting.HardwareTargetingSettings]
+TargetedHardwareClass=Desktop
+AppliedTargetedHardwareClass=Desktop
+DefaultGraphicsPerformance=Maximum
+AppliedDefaultGraphicsPerformance=Maximum
+
+[/Script/WorldPartitionEditor.WorldPartitionEditorSettings]
+bEnableWorldPartition=True
+bEnableConversionPrompt=True
+CommandletClass=Class'/Script/UnrealEd.WorldPartitionConvertCommandlet
+
+[/Script/EngineSettings.GameMapsSettings]
+EditorStartupMap=/Game/Maps/2kmTemplate.2kmTemplate
+GameDefaultMap=/Game/Maps/2kmTemplate.2kmTemplate
+
diff --git a/Templates/TP_OpenWorld/Config/DefaultGame.ini b/Templates/TP_OpenWorld/Config/DefaultGame.ini
new file mode 100644
--- /dev/null
+++ b/Templates/TP_OpenWorld/Config/DefaultGame.ini
@@ -0,0 +1,2 @@
+[/Script/EngineSettings.GeneralProjectSettings]
+ProjectID=3A88BA6E4B1578A555FFB5BEAEAB67C4
diff --git a/Templates/TP_OpenWorld/TP_OpenWorld.uproject b/Templates/TP_OpenWorld/TP_OpenWorld.uproject
new file mode 100644
--- /dev/null
+++ b/Templates/TP_OpenWorld/TP_OpenWorld.uproject
@@ -0,0 +1,6 @@
+{
+	"FileVersion": 3,
+	"EngineAssociation": "",
+	"Category": "",
+	"Description": ""
+}
\ No newline at end of file
